NavCorridor: renamed draf to taper and fixed corridor building when path locations are left ouside the corridor.

diff --git a/Engine/Plugins/Runtime/NavCorridor/Source/NavCorridor/Private/NavCorridor.cpp b/Engine/Plugins/Runtime/NavCorridor/Source/NavCorridor/Private/NavCorridor.cpp
--- a/Engine/Plugins/Runtime/NavCorridor/Source/NavCorridor/Private/NavCorridor.cpp
+++ b/Engine/Plugins/Runtime/NavCorridor/Source/NavCorridor/Private/NavCorridor.cpp
@@ -271,24 +271,31 @@ namespace UE::NavCorridor::Private
 	}
 
 	/**
-	 * Finds clusters of connected segments and adds draft segments at the beginning and end of the clusters.
-	 * A draft segment is an angled segment that makes the obstacle edges tapered. These tapers helps to remove small sectors from the
+	 * Ensures that segments slopes are within specified range. Fixes internal edge slopes, and adds taper segments at the end extrema of objects.
+	 * A taper segment is an angled segment that makes the obstacle edges tapered. These tapers helps to remove small sectors from the
 	 * final corridor (steep angle results in a narrow sector that cannot be removed), and it helps to prevent local traps which
 	 * can cause problems with avoidance, etc.
+	 *
+	 *         o----o
+	 *        /:     \
+	 *     o---o      \
+	 *    /            \
+	 * ../..............\.
+	 *
 	 * @param Segments obstacle segments
 	 * @param EdgeU U value of the edge we're currently processing (left = 0.0, right = 1.0)
-	 * @param DraftV How much draft to add to a vertex based on distance from corridor edge (EdgeU) 
+	 * @param SlopeOffsetV How much V is offset based on change in U. 
 	 */
-	static void AddDraftSegments(TArray<FUVSegment>& Segments, const FVector::FReal EdgeU, const FVector::FReal DraftV)
+	static void AddTaperSegments(TArray<FUVSegment>& Segments, const FVector::FReal EdgeU, const FVector::FReal SlopeOffsetV)
 	{
-		CSV_SCOPED_TIMING_STAT_EXCLUSIVE(NavCorridor_AddDraftSegments);
+		CSV_SCOPED_TIMING_STAT_EXCLUSIVE(NavCorridor_AddTaperSegments);
 
 		using FReal = FVector::FReal;
 
-		TArray<FUVSegment> DraftSegments;
-		DraftSegments.Reserve(Segments.Num());
+		TArray<FUVSegment> SlopeSegments;
+		SlopeSegments.Reserve(Segments.Num());
 		
-		const FReal MinSlope = 1.0 / DraftV;
+		const FReal MinSlope = 1.0 / SlopeOffsetV;
 
 		// Find end points of edge chains. The edges are already sorted based on Y (V) coordinate.  
 		TArray<bool> StartConnected;
@@ -316,7 +323,7 @@ namespace UE::NavCorridor::Private
 			}
 		}
 
-		// Correct segment draft
+		// Correct segment slopes
 		for (int32 SegmentIndex = 0; SegmentIndex < Segments.Num(); SegmentIndex++)
 		{
 			FUVSegment& CurrSegment = Segments[SegmentIndex];
@@ -333,11 +340,11 @@ namespace UE::NavCorridor::Private
 				if (StartDistance > EndDistance)
 				{
 					// Start is higher, keep start.
-					CurrSegment.EndUV = FVector2D(EdgeU, CurrSegment.StartUV.Y + StartDistance * DraftV);
+					CurrSegment.EndUV = FVector2D(EdgeU, CurrSegment.StartUV.Y + StartDistance * SlopeOffsetV);
 				}
 				else
 				{
-					CurrSegment.StartUV = FVector2D(EdgeU, CurrSegment.EndUV.Y - EndDistance * DraftV);
+					CurrSegment.StartUV = FVector2D(EdgeU, CurrSegment.EndUV.Y - EndDistance * SlopeOffsetV);
 				}
 				StartDistance = FMath::Abs(EdgeU - CurrSegment.StartUV.X);
 				EndDistance = FMath::Abs(EdgeU - CurrSegment.EndUV.X);
@@ -346,15 +353,15 @@ namespace UE::NavCorridor::Private
 			// Add a draft segment if the edge is start or end of and edge chain and not connected to edge.
 			if (!StartConnected[SegmentIndex] && StartDistance > UE_KINDA_SMALL_NUMBER)
 			{
-				DraftSegments.Emplace(FVector2D(EdgeU, CurrSegment.StartUV.Y - StartDistance * DraftV), CurrSegment.StartUV);
+				SlopeSegments.Emplace(FVector2D(EdgeU, CurrSegment.StartUV.Y - StartDistance * SlopeOffsetV), CurrSegment.StartUV);
 			}
 			if (!EndConnected[SegmentIndex] && EndDistance > UE_KINDA_SMALL_NUMBER)
 			{
-				DraftSegments.Emplace(CurrSegment.EndUV, FVector2D(EdgeU, CurrSegment.EndUV.Y + EndDistance * DraftV));
+				SlopeSegments.Emplace(CurrSegment.EndUV, FVector2D(EdgeU, CurrSegment.EndUV.Y + EndDistance * SlopeOffsetV));
 			}
 		}
 
-		Segments.Append(DraftSegments);
+		Segments.Append(SlopeSegments);
 		
 		Algo::Sort(Segments, [](const FUVSegment& A, const FUVSegment& B) { return A.StartUV.Y < B.StartUV.Y; });
 	}
@@ -376,7 +383,7 @@ namespace UE::NavCorridor::Private
 	 * @param OutSegmentIndex Hit segment index, or INDEX_NONE if no hit.
 	 * @return height (U value) at location SampleAtV
 	 */
-	static FVector::FReal SampleSegments(FSampleSegmentsIterator& Iter, const FVector::FReal SampleAtV, const FVector::FReal EdgeU, const TConstArrayView<FUVSegment> Segments, const TConstArrayView<FVector> Quad, int32& OutSegmentIndex)
+	static FVector::FReal SampleSegmentsHeight(FSampleSegmentsIterator& Iter, const FVector::FReal SampleAtV, const FVector::FReal EdgeU, const TConstArrayView<FUVSegment> Segments, const TConstArrayView<FVector> Quad, int32& OutSegmentIndex)
 	{
 		using FReal = FVector::FReal;
 
@@ -454,7 +461,7 @@ namespace UE::NavCorridor::Private
 	 * @param Quad bilinear patch used for mapping between world and UV space
 	 * @param OutUVPoints Hit segment index, or INDEX_NONE if no hit.
 	 */
-	static void CalculateHull(const TConstArrayView<FUVSegment> Segments, const FVector::FReal EdgeU, const TConstArrayView<FVector> Quad, TArray<FVector2D>& OutUVPoints)
+	static void CalculateHullPolyline(const TConstArrayView<FUVSegment> Segments, const FVector::FReal EdgeU, const TConstArrayView<FVector> Quad, TArray<FVector2D>& OutUVPoints)
 	{
 		CSV_SCOPED_TIMING_STAT_EXCLUSIVE(NavCorridor_CalculateHull);
 
@@ -538,7 +545,7 @@ namespace UE::NavCorridor::Private
 		for (const FReal V : SampleLocations)
 		{
 			int32 SegmentIndex = INDEX_NONE;
-			const FReal SampledU = SampleSegments(SampleIter, V, EdgeU, Segments, Quad, SegmentIndex);
+			const FReal SampledU = SampleSegmentsHeight(SampleIter, V, EdgeU, Segments, Quad, SegmentIndex);
 			const FVector2D UV(SampledU, V);
 
 			// If a third point is sampled on the same segment, do not add new point, just move the end point.
@@ -981,9 +988,99 @@ namespace UE::NavCorridor::Private
 	}
 
 	/**
-	 * Finds how much the end portal is visible from the start portal path location.
+	 * Pulls string between path locations on given portal indices.
+	 * Start and end points will stay in place, and all the points in between will be on shortest path along the corridor.
 	 * @param Portals portals defining the corridor
 	 * @param StartIndex start index of the string pull
+	 * @param EndIndex end index of the string pull
+	 */
+	static void StringPull(TArrayView<FNavCorridorPortal> Portals, const int32 StartIndex, const int32 EndIndex)
+	{
+		check(Portals.IsValidIndex(StartIndex));
+		check(Portals.IsValidIndex(EndIndex));
+		check(StartIndex <= EndIndex);
+		
+		FVector PortalApex = Portals[StartIndex].Location;
+		FVector PortalLeft = Portals[StartIndex].Left;
+		FVector PortalRight = Portals[StartIndex].Right;
+
+		int32 ApexIndex = StartIndex;
+		int32 LeftIndex = StartIndex;
+		int32 RightIndex = StartIndex;
+		
+		for (int32 Index = StartIndex + 1; Index <= EndIndex; Index++)
+		{
+			const FNavCorridorPortal& CurrPortal = Portals[Index];
+			const FVector Left = Index == EndIndex ? CurrPortal.Location : CurrPortal.Left;
+			const FVector Right = Index == EndIndex ? CurrPortal.Location : CurrPortal.Right;
+			if (UE::AI::TriArea2D(PortalApex, PortalRight, Right) >= 0.0)
+			{
+				if (PortalApex.Equals(PortalRight) || UE::AI::TriArea2D(PortalApex, PortalLeft, Right) < 0.0)
+				{
+					// Tighten the funnel
+					PortalRight = Right;
+					RightIndex = Index;
+				}
+				else
+				{
+					// Right over left, insert left to path and restart scan from portal left point.
+					Portals[LeftIndex].Location = PortalLeft;
+					Portals[LeftIndex].bIsPathCorner = true;
+					InterpolateInBetweenLocations(Portals, ApexIndex, LeftIndex);
+					// Make current left the new apex.
+					PortalApex = PortalLeft;
+					ApexIndex = LeftIndex;
+					// Reset portal
+					PortalLeft = PortalApex;
+					PortalRight = PortalApex;
+					LeftIndex = ApexIndex;
+					RightIndex = ApexIndex;
+					// Restart scan
+					Index = ApexIndex;
+					continue;
+				}
+			}
+			if (UE::AI::TriArea2D(PortalApex, PortalLeft, Left) <= 0.0)
+			{
+				if (PortalApex.Equals(PortalLeft) || UE::AI::TriArea2D(PortalApex, PortalRight, Left) > 0.0)
+				{
+					// Tighten the funnel
+					PortalLeft = Left;
+					LeftIndex = Index;
+				}
+				else
+				{
+					// Left over right, insert right to path and restart scan from portal right point.
+					Portals[RightIndex].Location = PortalRight;
+					Portals[RightIndex].bIsPathCorner = true;
+					InterpolateInBetweenLocations(Portals, ApexIndex, RightIndex);
+					// Make current left the new apex.
+					PortalApex = PortalRight;
+					ApexIndex = RightIndex;
+					// Reset portal
+					PortalLeft = PortalApex;
+					PortalRight = PortalApex;
+					LeftIndex = ApexIndex;
+					RightIndex = ApexIndex;
+					// Restart scan
+					Index = ApexIndex;
+					continue;
+				}
+			}
+		}
+	}
+
+	/** Projects the 'Location' of the portal on the portal segment. */
+	static void ProjectLocationOnPortal(FNavCorridorPortal& InOutPortal)
+	{
+		const double LocationU = UE::AI::ProjectPointOnSegment2D(InOutPortal.Location, InOutPortal.Left, InOutPortal.Right);
+		InOutPortal.Location = FMath::Lerp(InOutPortal.Left, InOutPortal.Right, LocationU);
+	}
+
+	/**
+	 * Finds how much the end portal is visible from the start portal path location.
+	 * @param Portals portals defining the corridor
+	 * @param StartIndex start index of the visibility check
 	 * @param EndIndex end index of the interpolation
 	 * @return Visibility min and max value in range [0..1] on the end portal between left and right portal points.  
 	 */
@@ -1029,7 +1126,7 @@ namespace UE::NavCorridor::Private
 	/**
 	 * Finds how much the start portal is visible from the end portal path location.
 	 * @param Portals portals defining the corridor
-	 * @param StartIndex start index of the string pull
+	 * @param StartIndex start index of the visibility check
 	 * @param EndIndex end index of the interpolation
 	 * @return Visibility min and max value in range [0..1] on the start portal between left and right portal points.  
 	 */
@@ -1076,7 +1173,7 @@ namespace UE::NavCorridor::Private
 	/**
 	 * Finds the next portal marked as corner starting from specified index.
 	 * @param Portals portals defining the corridor
-	 * @param StartIndex start index of the string pull
+	 * @param StartIndex start index of the query
 	 * @return Index of next corner points, or INDEX_NONE if not found.   
 	 */
 	static int32 FindNextCorner(TArrayView<FNavCorridorPortal> Portals, const int32 StartIndex)
@@ -1141,6 +1238,8 @@ void FNavCorridor::BuildFromPathPoints(const ANavigationData* NavData, FSharedCo
 	UE::NavCorridor::Private::CalculateSegmentDirections(PathPoints, SegmentDirs, SegmentLefts);
 	UE::NavCorridor::Private::CalculateCorridorPortals(PathPoints, SegmentLefts, Params.Width, InitialPortals);
 
+	const FReal TaperLength = FMath::Tan(FMath::DegreesToRadians(Params.ObstacleTaperAngle)) * Params.Width;
+	
 	check(InitialPortals.Num() == PathPoints.Num());
 
 	TArray<FVector> Edges;
@@ -1177,28 +1276,27 @@ void FNavCorridor::BuildFromPathPoints(const ANavigationData* NavData, FSharedCo
 
 		UE::NavCorridor::Private::ClipEdgesToSector(SegBasePos, SegDir, SegLeft, Quad, Edges, LeftUVSegments, RightUVSegments);
 
-		// Process obstacle segments so that they do not have steep angles, by adding draft/taper segments at obstacle ends.
+		// Process obstacle segments so that they do not have steep angles, by adding taper segments at obstacle ends.
 		TArray<FVector2d> LeftUVPoints;
 		TArray<FVector2d> RightUVPoints;
 
 		const FReal LeftSideDistance = FVector::Distance(CurrPortal.Left, NextPortal.Left);
 		const FReal RightSideDistance = FVector::Distance(CurrPortal.Right, NextPortal.Right);
-		const FReal LeftDraftV = LeftSideDistance > UE_KINDA_SMALL_NUMBER ? (Params.ClusterDraftLength / LeftSideDistance) : 0.0;
-		const FReal RightDraftV = RightSideDistance > UE_KINDA_SMALL_NUMBER ? (Params.ClusterDraftLength / RightSideDistance) : 0.0;
+		const FReal LeftSlopeOffsetV = LeftSideDistance > UE_KINDA_SMALL_NUMBER ? (TaperLength / LeftSideDistance) : 0.0;
+		const FReal RightSlopeOffsetV = RightSideDistance > UE_KINDA_SMALL_NUMBER ? (TaperLength / RightSideDistance) : 0.0;
 
-		UE::NavCorridor::Private::AddDraftSegments(LeftUVSegments, 0.0, LeftDraftV);
-		UE::NavCorridor::Private::AddDraftSegments(RightUVSegments, 1.0, RightDraftV);
+		UE::NavCorridor::Private::AddTaperSegments(LeftUVSegments, 0.0, LeftSlopeOffsetV);
+		UE::NavCorridor::Private::AddTaperSegments(RightUVSegments, 1.0, RightSlopeOffsetV);
 
 		// Calculate monotonic polyline hull of each side. 
-		UE::NavCorridor::Private::CalculateHull(LeftUVSegments, 0.0, Quad, LeftUVPoints);
-		UE::NavCorridor::Private::CalculateHull(RightUVSegments, 1.0, Quad, RightUVPoints);
+		UE::NavCorridor::Private::CalculateHullPolyline(LeftUVSegments, 0.0, Quad, LeftUVPoints);
+		UE::NavCorridor::Private::CalculateHullPolyline(RightUVSegments, 1.0, Quad, RightUVPoints);
 		
 		// Combine the hulls and create portals at each vertex on left and right hull.
 		UE::NavCorridor::Private::DivideHullsIntoPortals(LeftUVPoints, RightUVPoints, Quad, CurrPortal.Location, NextPortal.Location, PathPointBaseIndex + PortalIndex,  Portals);
 	}
 
 	// The process above can leave short sectors in the corridor. The passes below remove common patterns of small sectors.
-	
 	if (Params.bSimplifyFlipPortals)
 	{
 		UE::NavCorridor::Private::SimplifyFlipPortals(Portals, Params.SimplifyEdgeThreshold);
@@ -1211,6 +1309,20 @@ void FNavCorridor::BuildFromPathPoints(const ANavigationData* NavData, FSharedCo
 	{
 		UE::NavCorridor::Private::SimplifyConcavePortals(Portals, Params.SimplifyEdgeThreshold, Params.LargeSectorThreshold);
 	}
+
+	if (Portals.Num() > 1)
+	{
+		// Ensure that the start and end locations are on portal
+		UE::NavCorridor::Private::ProjectLocationOnPortal(Portals[0]);
+		UE::NavCorridor::Private::ProjectLocationOnPortal(Portals.Last());
+		
+		// String pull the path again, might not be correct anymore after all the processing.
+		if (Portals.Num() > 2)
+		{
+			// String pull
+			UE::NavCorridor::Private::StringPull(Portals, 0, Portals.Num() - 1);
+		}
+	}
 }
 
 void FNavCorridor::OffsetPathLocationsFromWalls(const float Offset, bool bOffsetFirst, bool bOffsetLast)
diff --git a/Engine/Plugins/Runtime/NavCorridor/Source/NavCorridor/Public/NavCorridor.h b/Engine/Plugins/Runtime/NavCorridor/Source/NavCorridor/Public/NavCorridor.h
--- a/Engine/Plugins/Runtime/NavCorridor/Source/NavCorridor/Public/NavCorridor.h
+++ b/Engine/Plugins/Runtime/NavCorridor/Source/NavCorridor/Public/NavCorridor.h
@@ -13,45 +13,45 @@ struct NAVCORRIDOR_API FNavCorridorParams
 {
 	GENERATED_BODY()
 
+	/** Sets good default values for the params based on corridor width. */
 	void SetFromWidth(const float InWidth)
 	{
 		Width = InWidth;
-		ClusterDraftLength = Width * 0.5f;
 		SmallSectorThreshold = Width * 0.3f;
 		LargeSectorThreshold = Width;
 		SimplifyEdgeThreshold = Width * 0.1f;
 	}
 	
 	/** Width of the corridor to build */
-	UPROPERTY(EditAnywhere, Category=Pathfinding)
+	UPROPERTY(EditAnywhere, Category = "Pathfinding", meta=(ClampMin="10.0"))
 	float Width = 200.0f;
 
-	/** How much the outer edges of obstacles are tapered (per width of the corridor). This prevents small sectors and local traps.  */
-	UPROPERTY(EditAnywhere, Category=Pathfinding)
-	float ClusterDraftLength = 100.0f;
+	/** How much the outer edges of obstacles are tapered out. This prevents small sectors and local traps.  */
+	UPROPERTY(EditAnywhere, Category = "Pathfinding", meta=(ClampMin="10.0", ClampMax="45.0"))
+	float ObstacleTaperAngle = 30.0f;
 
 	/** Attempt to remove sectors narrower than this from the corridor. */
-	UPROPERTY(EditAnywhere, Category=Pathfinding)
+	UPROPERTY(EditAnywhere, Category = "Pathfinding", meta=(ClampMin="0.0"))
 	float SmallSectorThreshold = 60.0f;
 
 	/** Simplification is skipped if two neighbour sectors combined are longer than this. This ensures that long sectors do not lose volume due to simplification. */
-	UPROPERTY(EditAnywhere, Category=Pathfinding)
+	UPROPERTY(EditAnywhere, Category = "Pathfinding", meta=(ClampMin="0.0"))
 	float LargeSectorThreshold = 200.0f;
 
 	/** Corridor edge max simplification distance. */
-	UPROPERTY(EditAnywhere, Category=Pathfinding)
+	UPROPERTY(EditAnywhere, Category = "Pathfinding", meta=(ClampMin="0.0"))
 	float SimplifyEdgeThreshold = 20.0f;
 
 	/** If true do flip portals simplification. */
-	UPROPERTY(EditAnywhere, Category=Pathfinding)
+	UPROPERTY(EditAnywhere, Category = "Pathfinding")
 	bool bSimplifyFlipPortals = true;
 	
 	/** If true do convex portals simplification. */
-	UPROPERTY(EditAnywhere, Category=Pathfinding)
+	UPROPERTY(EditAnywhere, Category = "Pathfinding")
 	bool bSimplifyConvexPortals = true;
 	
 	/** If true do concave portals simplification. */
-	UPROPERTY(EditAnywhere, Category=Pathfinding)
+	UPROPERTY(EditAnywhere, Category = "Pathfinding")
 	bool bSimplifyConcavePortals = true;
 };
 
