NavCorridor: Compute the obstacle offset inside the funnel string pulling
This prevents pushing the path too close to the other side of the corridor and gives better results in general

diff --git a/Engine/Plugins/Runtime/NavCorridor/Source/NavCorridor/Private/NavCorridor.cpp b/Engine/Plugins/Runtime/NavCorridor/Source/NavCorridor/Private/NavCorridor.cpp
--- a/Engine/Plugins/Runtime/NavCorridor/Source/NavCorridor/Private/NavCorridor.cpp
+++ b/Engine/Plugins/Runtime/NavCorridor/Source/NavCorridor/Private/NavCorridor.cpp
@@ -11,6 +11,11 @@
 // If set to 1, parts of the process are logged into vis logger.
 #define NAV_CORRIDOR_DEBUG_DETAILS 0
 
+#if NAV_CORRIDOR_DEBUG_DETAILS
+DEFINE_LOG_CATEGORY_STATIC(LogNavCorridorQuads, Log, All);
+DEFINE_LOG_CATEGORY_STATIC(LogNavCorridorEdges, Log, All);
+#endif
+
 //-------------------------------------------------------
 // Private helper functions for the nav corridor.
 //-------------------------------------------------------
@@ -285,15 +290,15 @@ namespace UE::NavCorridor::Private
 					{
 						const FVector Offset(0,0, 15);
 						const FVector Offset2(0,0, 17);
-						UE_VLOG_SEGMENT_THICK(NavData, LogNavigation, Log, StartPos + Offset, EndPos + Offset2, FColor::Blue, 2, TEXT_EMPTY);
-						UE_VLOG_SEGMENT_THICK(NavData, LogNavigation, Log, Mid + Offset, Mid + Left * 10.0f + Offset, bAdded ? FColor::Blue : FColor::Red, 1, TEXT_EMPTY);
+						UE_VLOG_SEGMENT_THICK(NavData, LogNavCorridorEdges, Log, StartPos + Offset, EndPos + Offset2, FColor::Blue, 2, TEXT_EMPTY);
+						UE_VLOG_SEGMENT_THICK(NavData, LogNavCorridorEdges, Log, Mid + Offset, Mid + Left * 10.0f + Offset, bAdded ? FColor::Blue : FColor::Red, 1, TEXT_EMPTY);
 					}
 					else
 					{
 						const FVector Offset(0,0, 15);
 						const FVector Offset2(0,0, 17);
-						UE_VLOG_SEGMENT_THICK(NavData, LogNavigation, Log, StartPos + Offset, EndPos + Offset2, FColor::Green, 2, TEXT_EMPTY);
-						UE_VLOG_SEGMENT_THICK(NavData, LogNavigation, Log, Mid + Offset, Mid + Left * 10.0f + Offset, bAdded ? FColor::Green : FColor::Red, 1, TEXT_EMPTY);
+						UE_VLOG_SEGMENT_THICK(NavData, LogNavCorridorEdges, Log, StartPos + Offset, EndPos + Offset2, FColor::Green, 2, TEXT_EMPTY);
+						UE_VLOG_SEGMENT_THICK(NavData, LogNavCorridorEdges, Log, Mid + Offset, Mid + Left * 10.0f + Offset, bAdded ? FColor::Green : FColor::Red, 1, TEXT_EMPTY);
 					}
 				}
 #endif				
@@ -982,11 +987,11 @@ namespace UE::NavCorridor::Private
 						const FReal ApproxPortalDistance12 = ApproxDistanceSegmentSegment(Portal1.Left, Portal1.Right, Portal2.Left, Portal2.Right);
 						if (ApproxPortalDistance01 < ApproxPortalDistance12)
 						{
-							Portal0.bIsPathCorner |= Portal1.bIsPathCorner; 
+							Portal0.bIsPathCorner |= Portal1.bIsPathCorner;
 						}
 						else
 						{
-							Portal2.bIsPathCorner |= Portal1.bIsPathCorner; 
+							Portal2.bIsPathCorner |= Portal1.bIsPathCorner;
 						}
 					}
 
@@ -1021,22 +1026,61 @@ namespace UE::NavCorridor::Private
 		}
 	}
 
+	/**
+	 * Returns a portal's extremities pulled in from their original position with OffsetWidth. If the portal is
+	 * smaller than 2x width, the center point will be used
+	 * @param Portal The portal to get Left/Right from
+	 * @param OffsetWidth How much to offset from the original corridor boundaries
+	 * @param OutOffsettedLeft The offsetted left extremity of the portal
+	 * @param OutOffsettedRight The offsetted right extremity of the portal
+	 */
+	static void GetPortalLeftRightWithOffset(const FNavCorridorPortal& Portal, const float OffsetWidth, FVector& OutOffsettedLeft, FVector& OutOffsettedRight)
+	{
+		using FReal = FVector::FReal;
+
+		if (OffsetWidth <= UE_KINDA_SMALL_NUMBER)
+		{
+			OutOffsettedLeft = Portal.Left;
+			OutOffsettedRight = Portal.Right;
+			return;
+		}
+
+		const FVector PortalDirection = Portal.Right - Portal.Left;
+		const FReal PortalWidthSquared = PortalDirection.SizeSquared2D();
+		const FReal TwoOffsetSquared = (2.0f * OffsetWidth) * (2.0f * OffsetWidth);
+		if (PortalWidthSquared < TwoOffsetSquared)
+		{
+			const FVector Middle = Portal.Left + 0.5 * PortalDirection;
+			OutOffsettedLeft = Middle;
+			OutOffsettedRight = Middle;
+		}
+		else
+		{
+			const FReal PortalWidth = FMath::Sqrt(PortalWidthSquared);
+			const FVector Offset = (PortalDirection / PortalWidth) * OffsetWidth; // Assuming the portal is horizontal enough to ignore offset width that gets consumed along the Z axis
+			OutOffsettedLeft = Portal.Left + Offset;
+			OutOffsettedRight = Portal.Right - Offset;
+		}
+	}
+
 	/**
 	 * Pulls string between path locations on given portal indices.
 	 * Start and end points will stay in place, and all the points in between will be on shortest path along the corridor.
 	 * @param Portals portals defining the corridor
 	 * @param StartIndex start index of the string pull
 	 * @param EndIndex end index of the string pull
+	 * @param OffsetFromBoundaries optional offset that will push the path away from the portal extremities and towards the center. No effect when set to 0.0
 	 */
-	static void StringPull(TArrayView<FNavCorridorPortal> Portals, const int32 StartIndex, const int32 EndIndex)
+	static void StringPull(TArrayView<FNavCorridorPortal> Portals, const int32 StartIndex, const int32 EndIndex, const float OffsetFromBoundaries)
 	{
 		check(Portals.IsValidIndex(StartIndex));
 		check(Portals.IsValidIndex(EndIndex));
 		check(StartIndex <= EndIndex);
 		
 		FVector PortalApex = Portals[StartIndex].Location;
-		FVector PortalLeft = Portals[StartIndex].Left;
-		FVector PortalRight = Portals[StartIndex].Right;
+		FVector PortalLeft;
+		FVector PortalRight;
+		GetPortalLeftRightWithOffset(Portals[StartIndex], OffsetFromBoundaries, PortalLeft, PortalRight);
 
 		int32 ApexIndex = StartIndex;
 		int32 LeftIndex = StartIndex;
@@ -1045,8 +1089,18 @@ namespace UE::NavCorridor::Private
 		for (int32 Index = StartIndex + 1; Index <= EndIndex; Index++)
 		{
 			const FNavCorridorPortal& CurrPortal = Portals[Index];
-			const FVector Left = Index == EndIndex ? CurrPortal.Location : CurrPortal.Left;
-			const FVector Right = Index == EndIndex ? CurrPortal.Location : CurrPortal.Right;
+			FVector Left;
+			FVector Right;
+			if (Index == EndIndex)
+			{
+				Left = CurrPortal.Location;
+				Right = CurrPortal.Location;
+			}
+			else
+			{
+				GetPortalLeftRightWithOffset(Portals[Index], OffsetFromBoundaries, Left, Right);
+			}
+
 			if (UE::AI::TriArea2D(PortalApex, PortalRight, Right) >= 0.0)
 			{
 				if (PortalApex.Equals(PortalRight) || UE::AI::TriArea2D(PortalApex, PortalLeft, Right) < 0.0)
@@ -1301,11 +1355,13 @@ void FNavCorridor::BuildFromPathPoints(const FNavigationPath& Path, TConstArrayV
 			Edges.Reset();
 			NavData->GetPathSegmentBoundaryEdges(Path,  PathPoints[PortalIndex], PathPoints[PortalIndex+1], Quad, Edges, 0.1f, NavQueryFilter);
 
-#if NAV_CORRIDOR_DEBUG_DETAILS			
+#if NAV_CORRIDOR_DEBUG_DETAILS
+			const FColor RandomColor = FLinearColor::MakeFromHSV8(FRandomStream(PortalIndex).FRandRange(0, 255.0), 196, 196).ToFColor(false).WithAlpha(128);
+			UE_VLOG_CONVEXPOLY(NavData, LogNavCorridorQuads, Log, TArray<FVector>(Quad), RandomColor, TEXT_EMPTY);
 			for (int32 Index = 0; Index < Edges.Num(); Index += 2)
 			{
 				const FVector PolyOffset(0,0,15);
-				UE_VLOG_ARROW(NavData, LogNavigation, Log, Edges[Index] + PolyOffset, Edges[Index+1] + PolyOffset, FColor::Black, TEXT_EMPTY);
+				UE_VLOG_ARROW(NavData, LogNavCorridorEdges, Log, Edges[Index] + PolyOffset, Edges[Index+1] + PolyOffset, FColor::Black, TEXT_EMPTY);
 			}
 #endif			
 
@@ -1366,7 +1422,7 @@ void FNavCorridor::BuildFromPathPoints(const FNavigationPath& Path, TConstArrayV
 		if (Portals.Num() > 2)
 		{
 			// String pull
-			UE::NavCorridor::Private::StringPull(Portals, 0, Portals.Num() - 1);
+			UE::NavCorridor::Private::StringPull(Portals, 0, Portals.Num() - 1, Params.PathOffsetFromBoundaries);
 		}
 	}
 }
diff --git a/Engine/Plugins/Runtime/NavCorridor/Source/NavCorridor/Public/NavCorridor.h b/Engine/Plugins/Runtime/NavCorridor/Source/NavCorridor/Public/NavCorridor.h
--- a/Engine/Plugins/Runtime/NavCorridor/Source/NavCorridor/Public/NavCorridor.h
+++ b/Engine/Plugins/Runtime/NavCorridor/Source/NavCorridor/Public/NavCorridor.h
@@ -27,6 +27,10 @@ struct NAVCORRIDOR_API FNavCorridorParams
 	UPROPERTY(EditAnywhere, Category = "Pathfinding", meta=(ClampMin="10.0"))
 	float Width = 200.0f;
 
+	/** Distance to maintain between the path and the corridor boundaries. If the corridor is too tight, the offset will be reduced to the available space. */
+	UPROPERTY(EditAnywhere, Category = "Pathfinding", meta = (ClampMin = "0.0"))
+	float PathOffsetFromBoundaries = 0.f;
+
 	/** How much the outer edges of obstacles are tapered out. This prevents small sectors and local traps.  */
 	UPROPERTY(EditAnywhere, Category = "Pathfinding", meta=(ClampMin="10.0", ClampMax="45.0"))
 	float ObstacleTaperAngle = 30.0f;
@@ -135,7 +139,7 @@ struct NAVCORRIDOR_API FNavCorridor : public TSharedFromThis<FNavCorridor>
 	/** @return target vector that is visible from Source. */
 	FVector ConstrainVisibility(const FNavCorridorLocation& PathLocation, const FVector Source, const FVector Target, const float ForceLookAheadDistance = 0.0f) const;
 
-	/** @return true if the segment hints the corridor boundary. */
+	/** @return true if the segment hits the corridor boundary. */
 	bool HitTest(const FVector SegmentStart, const FVector SegmentEnd, double& HitT);
 
 	/** Portal points defining the corridor. */
