NavCorridor:
*Added a "pin" option for points passed to the NavCorridor so they don't get removed or string pulled by the algorithm. This is useful when the path contains points that we absolutely want to keep (e.g. hand-placed patrol path points).
*Added more VisualLogger information for NavCorridors.

#rnx

diff --git a/Engine/Plugins/Runtime/NavCorridor/Source/NavCorridor/Private/NavCorridor.cpp b/Engine/Plugins/Runtime/NavCorridor/Source/NavCorridor/Private/NavCorridor.cpp
--- a/Engine/Plugins/Runtime/NavCorridor/Source/NavCorridor/Private/NavCorridor.cpp
+++ b/Engine/Plugins/Runtime/NavCorridor/Source/NavCorridor/Private/NavCorridor.cpp
@@ -1,20 +1,21 @@
 ﻿// Copyright Epic Games, Inc. All Rights Reserved.
 
 #include "NavCorridor.h"
+
+#include "AI/Navigation/NavAgentInterface.h"
 #include "GeomUtils.h"
 #include "NavigationData.h"
+#include "NavMesh/NavMeshPath.h"
 #include "ProfilingDebugging/CsvProfiler.h"
-
+#include "VisualLogger/VisualLogger.h"
 
 #include UE_INLINE_GENERATED_CPP_BY_NAME(NavCorridor)
 
-// If set to 1, parts of the process are logged into vis logger.
-#define NAV_CORRIDOR_DEBUG_DETAILS 0
-
-#if NAV_CORRIDOR_DEBUG_DETAILS
+#if ENABLE_VISUAL_LOG
 DEFINE_LOG_CATEGORY_STATIC(LogNavCorridorQuads, Log, All);
 DEFINE_LOG_CATEGORY_STATIC(LogNavCorridorEdges, Log, All);
-#endif
+DEFINE_LOG_CATEGORY_STATIC(LogNavCorridorPath, Log, All);
+#endif // ENABLE_VISUAL_LOG
 
 //-------------------------------------------------------
 // Private helper functions for the nav corridor.
@@ -22,6 +23,11 @@ DEFINE_LOG_CATEGORY_STATIC(LogNavCorridorEdges, Log, All);
 
 namespace UE::NavCorridor::Private
 {
+#if ENABLE_VISUAL_LOG
+	static bool bEnableNavCorridorVLog = false;
+	static FAutoConsoleVariableRef CVarNavCorridorDebugEnableVisualLogging(TEXT("ai.nav.EnableNavCorridorVLog"), bEnableNavCorridorVLog, TEXT("If set, NavCorridor debug information will be logged in the Visual Logger."), ECVF_Cheat);
+#endif // ENABLE_VISUAL_LOG
+
 	/** Epsilon used when testing if 3 points for a convex corner.
 	 * The check is done using signed distance of the corner points to the segment from first to last points.
 	 * This is how much it can be concave. */
@@ -95,7 +101,7 @@ namespace UE::NavCorridor::Private
 	 * @param Width width of the corridor to generate.
 	 * @param OutPortals resulting corridor portals.
 	 */
-	static void CalculateCorridorPortals(const TConstArrayView<FNavPathPoint> Points, const TArrayView<FVector> SegmentLefts, const FVector::FReal Width, TArray<FNavCorridorPortal>& OutPortals)
+	static void CalculateCorridorPortals(const TConstArrayView<FNavPathPoint> Points, const TConstArrayView<int32> PinnedPointsIndices, const TArrayView<FVector> SegmentLefts, const FVector::FReal Width, TArray<FNavCorridorPortal>& OutPortals)
 	{
 		using FReal = FVector::FReal;
 		TArray<FVector> PointMiterDirs;
@@ -191,6 +197,7 @@ namespace UE::NavCorridor::Private
 			Portal.Left = Base + Delta * PortalLeftU[PointIndex];
 			Portal.Right = Base + Delta * PortalRightU[PointIndex];
 			Portal.Location = Loc;
+			Portal.bIsPinnedLocation = PinnedPointsIndices.Contains(PointIndex);
 		}
 	}
 
@@ -204,7 +211,7 @@ namespace UE::NavCorridor::Private
 	 * @param OutLeftSegments edges in UV space left of the segment 
 	 * @param OutRightSegments edges in UV space right of the segment
 	 */
-	static void ClipEdgesToSector(const ANavigationData* NavData, const FVector SegStart, const FVector SegForward, const FVector SegLeft, TConstArrayView<FVector> Quad, TConstArrayView<FVector> Edges,
+	static void ClipEdgesToSector(const UObject* PathQuerier, const FVector SegStart, const FVector SegForward, const FVector SegLeft, TConstArrayView<FVector> Quad, TConstArrayView<FVector> Edges,
 						  TArray<FUVSegment>& OutLeftSegments, TArray<FUVSegment>& OutRightSegments)
 	{
 		CSV_SCOPED_TIMING_STAT_EXCLUSIVE(NavCorridor_ClipEdgesToSector);
@@ -280,7 +287,8 @@ namespace UE::NavCorridor::Private
 					}
 				}
 
-#if NAV_CORRIDOR_DEBUG_DETAILS				
+#if ENABLE_VISUAL_LOG
+				if (UE::NavCorridor::Private::bEnableNavCorridorVLog)
 				{
 					const FVector Dir = EndPos - StartPos;
 					const FVector Left = FVector::CrossProduct(Dir, FVector::UpVector).GetSafeNormal();
@@ -290,18 +298,18 @@ namespace UE::NavCorridor::Private
 					{
 						const FVector Offset(0,0, 15);
 						const FVector Offset2(0,0, 17);
-						UE_VLOG_SEGMENT_THICK(NavData, LogNavCorridorEdges, Log, StartPos + Offset, EndPos + Offset2, FColor::Blue, 2, TEXT_EMPTY);
-						UE_VLOG_SEGMENT_THICK(NavData, LogNavCorridorEdges, Log, Mid + Offset, Mid + Left * 10.0f + Offset, bAdded ? FColor::Blue : FColor::Red, 1, TEXT_EMPTY);
+						UE_VLOG_SEGMENT_THICK(PathQuerier, LogNavCorridorEdges, Log, StartPos + Offset, EndPos + Offset2, FColor::Blue, 2, TEXT_EMPTY);
+						UE_VLOG_SEGMENT_THICK(PathQuerier, LogNavCorridorEdges, Log, Mid + Offset, Mid + Left * 10.0f + Offset, bAdded ? FColor::Blue : FColor::Red, 1, TEXT_EMPTY);
 					}
 					else
 					{
 						const FVector Offset(0,0, 15);
 						const FVector Offset2(0,0, 17);
-						UE_VLOG_SEGMENT_THICK(NavData, LogNavCorridorEdges, Log, StartPos + Offset, EndPos + Offset2, FColor::Green, 2, TEXT_EMPTY);
-						UE_VLOG_SEGMENT_THICK(NavData, LogNavCorridorEdges, Log, Mid + Offset, Mid + Left * 10.0f + Offset, bAdded ? FColor::Green : FColor::Red, 1, TEXT_EMPTY);
+						UE_VLOG_SEGMENT_THICK(PathQuerier, LogNavCorridorEdges, Log, StartPos + Offset, EndPos + Offset2, FColor::Green, 2, TEXT_EMPTY);
+						UE_VLOG_SEGMENT_THICK(PathQuerier, LogNavCorridorEdges, Log, Mid + Offset, Mid + Left * 10.0f + Offset, bAdded ? FColor::Green : FColor::Red, 1, TEXT_EMPTY);
 					}
 				}
-#endif				
+#endif // ENABLE_VISUAL_LOG
 			}
 		}
 
@@ -666,13 +674,13 @@ namespace UE::NavCorridor::Private
 	 * @param LeftUVPoints left hull of the corridor 
 	 * @param RightUVPoints right hull of the corridor
 	 * @param Quad bilinear patch used for mapping between world and UV space
-	 * @param PathSegmentStart original path segment start location
-	 * @param PathSegmentEnd original path segment end location
+	 * @param PathSegmentStartPortal original path segment start portal
+	 * @param PathSegmentEndPortal original path segment end portal
 	 * @param PathPointBaseIndex original path segment start point index
 	 * @param OutPortals resulting portals
 	 */
 	static void DivideHullsIntoPortals(const TConstArrayView<FVector2D> LeftUVPoints, const TConstArrayView<FVector2D> RightUVPoints, const TConstArrayView<FVector> Quad,
-										const FVector PathSegmentStart, const FVector PathSegmentEnd, const int32 PathPointBaseIndex,
+										const FNavCorridorPortal& PathSegmentStartPortal, const FNavCorridorPortal& PathSegmentEndPortal, const int32 PathPointBaseIndex,
 										TArray<FNavCorridorPortal>& OutPortals)
 	{
 		CSV_SCOPED_TIMING_STAT_EXCLUSIVE(NavCorridor_DivideHullsIntoPortals);
@@ -706,7 +714,8 @@ namespace UE::NavCorridor::Private
 			
 		Algo::Sort(SampleLocation, [](const FReal A, const FReal B) { return A < B; });
 
-		const int32 First = OutPortals.Num() > 0 ? 1 : 0;
+		// Only add the StartPortal on the first call to this function (OutPortals is empty) since the EndPortal will become the StartPortal for the next iteration.
+		const bool bAddStartPortal = OutPortals.IsEmpty();
 
 		FInterpolateIterator LeftIter;
 		FInterpolateIterator RightIter;
@@ -727,7 +736,7 @@ namespace UE::NavCorridor::Private
 		}
 
 		// Interpolate and add rest of the portals.
-		for (int32 DivIndex = First; DivIndex < SampleLocation.Num(); DivIndex++)
+		for (int32 DivIndex = bAddStartPortal ? 0 : 1; DivIndex < SampleLocation.Num(); DivIndex++)
 		{
 			const FReal V = SampleLocation[DivIndex];
 
@@ -737,10 +746,20 @@ namespace UE::NavCorridor::Private
 			FNavCorridorPortal& Portal = OutPortals.AddDefaulted_GetRef();
 			Portal.Left = LeftPos;
 			Portal.Right = RightPos;
-			Portal.Location = FMath::Lerp(PathSegmentStart, PathSegmentEnd, V);
-			Portal.bIsPathCorner = (DivIndex == 0) || (DivIndex == (SampleLocation.Num() - 1));
-			Portal.PathPointIndex = PathPointBaseIndex;
+			Portal.Location = FMath::Lerp(PathSegmentStartPortal.Location, PathSegmentEndPortal.Location, V);
+			Portal.PathPointIndex = (int16)PathPointBaseIndex;
 		}
+
+		if (bAddStartPortal)
+		{
+			FNavCorridorPortal& FirstPortal = OutPortals[0];
+			FirstPortal.bIsPathCorner = true;
+			FirstPortal.bIsPinnedLocation = PathSegmentStartPortal.bIsPinnedLocation;
+		}
+
+		FNavCorridorPortal& LastPortal = OutPortals.Last();
+		LastPortal.bIsPathCorner = true;
+		LastPortal.bIsPinnedLocation = PathSegmentEndPortal.bIsPinnedLocation;
 	}
 
 	/**
@@ -775,6 +794,12 @@ namespace UE::NavCorridor::Private
 			const FNavCorridorPortal& Portal2 = Portals[Index + 1];
 			const FNavCorridorPortal& Portal3 = Portals[Index + 2];
 
+			// Skip simplification for pinned locations, they need to remain in the final list.
+			if (Portal2.bIsPinnedLocation)
+			{
+				continue;
+			}
+
 			// Only try to remove short sectors.
 			const FReal ApproxPortalDistance = ApproxDistanceSegmentSegment(Portal1.Left, Portal1.Right, Portal2.Left, Portal2.Right);
 			if (ApproxPortalDistance < SmallDistanceThreshold)
@@ -857,6 +882,12 @@ namespace UE::NavCorridor::Private
 			const FNavCorridorPortal& Portal1 = Portals[Index];
 			FNavCorridorPortal& Portal2 = Portals[Index + 1];
 
+			// Skip simplification for pinned locations, they need to remain in the final list.
+			if (Portal1.bIsPinnedLocation)
+			{
+				continue;
+			}
+
 			const FReal ApproxPortalDistance02 = ApproxDistanceSegmentSegment(Portal0.Left, Portal0.Right, Portal2.Left, Portal2.Right);
 			if (ApproxPortalDistance02 > LargeSectorThreshold)
 			{
@@ -929,6 +960,12 @@ namespace UE::NavCorridor::Private
 			FNavCorridorPortal& Portal1 = Portals[Index];
 			FNavCorridorPortal& Portal2 = Portals[Index + 1];
 
+			// Skip simplification for pinned locations, they need to remain in the final list.
+			if (Portal1.bIsPinnedLocation)
+			{
+				continue;
+			}
+
 			const FReal ApproxPortalDistance02 = ApproxDistanceSegmentSegment(Portal0.Left, Portal0.Right, Portal2.Left, Portal2.Right);
 			if (ApproxPortalDistance02 > LargeSectorThreshold)
 			{
@@ -1091,7 +1128,7 @@ namespace UE::NavCorridor::Private
 			const FNavCorridorPortal& CurrPortal = Portals[Index];
 			FVector Left;
 			FVector Right;
-			if (Index == EndIndex)
+			if (Index == EndIndex || CurrPortal.bIsPinnedLocation)
 			{
 				Left = CurrPortal.Location;
 				Right = CurrPortal.Location;
@@ -1326,7 +1363,7 @@ void FNavCorridor::BuildFromPathPoints(const FNavigationPath& Path, TConstArrayV
 	InitialPortals.Reserve(PathPoints.Num());
 
 	UE::NavCorridor::Private::CalculateSegmentDirections(PathPoints, SegmentDirs, SegmentLefts);
-	UE::NavCorridor::Private::CalculateCorridorPortals(PathPoints, SegmentLefts, Params.Width, InitialPortals);
+	UE::NavCorridor::Private::CalculateCorridorPortals(PathPoints, Params.PinnedPointsIndices, SegmentLefts, Params.Width, InitialPortals);
 
 	const FReal TaperLength = FMath::Tan(FMath::DegreesToRadians(Params.ObstacleTaperAngle)) * Params.Width;
 	
@@ -1335,6 +1372,7 @@ void FNavCorridor::BuildFromPathPoints(const FNavigationPath& Path, TConstArrayV
 	TArray<FVector> Edges;
 	Edges.Reserve(20); // Usually there are less than 10 segments. 
 
+	const UObject* PathQuerier = Path.GetQuerier();
 	for (int32 PortalIndex = 0; PortalIndex < InitialPortals.Num() - 1; PortalIndex++)
 	{
 		FNavCorridorPortal& CurrPortal = InitialPortals[PortalIndex];
@@ -1355,15 +1393,18 @@ void FNavCorridor::BuildFromPathPoints(const FNavigationPath& Path, TConstArrayV
 			Edges.Reset();
 			NavData->GetPathSegmentBoundaryEdges(Path,  PathPoints[PortalIndex], PathPoints[PortalIndex+1], Quad, Edges, 0.1f, NavQueryFilter);
 
-#if NAV_CORRIDOR_DEBUG_DETAILS
-			const FColor RandomColor = FLinearColor::MakeFromHSV8(FRandomStream(PortalIndex).FRandRange(0, 255.0), 196, 196).ToFColor(false).WithAlpha(128);
-			UE_VLOG_CONVEXPOLY(NavData, LogNavCorridorQuads, Log, TArray<FVector>(Quad), RandomColor, TEXT_EMPTY);
-			for (int32 Index = 0; Index < Edges.Num(); Index += 2)
+#if ENABLE_VISUAL_LOG
+			if (UE::NavCorridor::Private::bEnableNavCorridorVLog)
 			{
-				const FVector PolyOffset(0,0,15);
-				UE_VLOG_ARROW(NavData, LogNavCorridorEdges, Log, Edges[Index] + PolyOffset, Edges[Index+1] + PolyOffset, FColor::Black, TEXT_EMPTY);
+				const FColor RandomColor = FLinearColor::MakeFromHSV8(FRandomStream(PortalIndex).FRandRange(0, 255.0), 196, 196).ToFColor(false).WithAlpha(128);
+				UE_VLOG_CONVEXPOLY(PathQuerier, LogNavCorridorQuads, Log, TArray<FVector>(Quad), RandomColor, TEXT_EMPTY);
+				for (int32 Index = 0; Index < Edges.Num(); Index += 2)
+				{
+					const FVector PolyOffset(0, 0, 15);
+					UE_VLOG_ARROW(PathQuerier, LogNavCorridorEdges, Log, Edges[Index] + PolyOffset, Edges[Index + 1] + PolyOffset, FColor::Black, TEXT_EMPTY);
+				}
 			}
-#endif			
+#endif // ENABLE_VISUAL_LOG
 
 		}
 
@@ -1376,7 +1417,7 @@ void FNavCorridor::BuildFromPathPoints(const FNavigationPath& Path, TConstArrayV
 		const FVector SegLeft = SegmentLefts[PortalIndex];
 		const FVector SegDir = SegmentDirs[PortalIndex];
 
-		UE::NavCorridor::Private::ClipEdgesToSector(NavData, SegBasePos, SegDir, SegLeft, Quad, Edges, LeftUVSegments, RightUVSegments);
+		UE::NavCorridor::Private::ClipEdgesToSector(PathQuerier, SegBasePos, SegDir, SegLeft, Quad, Edges, LeftUVSegments, RightUVSegments);
 
 		// Process obstacle segments so that they do not have steep angles, by adding taper segments at obstacle ends.
 		TArray<FVector2d> LeftUVPoints;
@@ -1395,7 +1436,7 @@ void FNavCorridor::BuildFromPathPoints(const FNavigationPath& Path, TConstArrayV
 		UE::NavCorridor::Private::CalculateHullPolyline(RightUVSegments, 1.0, Quad, RightUVPoints);
 		
 		// Combine the hulls and create portals at each vertex on left and right hull.
-		UE::NavCorridor::Private::DivideHullsIntoPortals(LeftUVPoints, RightUVPoints, Quad, CurrPortal.Location, NextPortal.Location, PathPointBaseIndex + PortalIndex,  Portals);
+		UE::NavCorridor::Private::DivideHullsIntoPortals(LeftUVPoints, RightUVPoints, Quad, CurrPortal, NextPortal, PathPointBaseIndex + PortalIndex, Portals);
 	}
 
 	// The process above can leave short sectors in the corridor. The passes below remove common patterns of small sectors.
@@ -1425,6 +1466,50 @@ void FNavCorridor::BuildFromPathPoints(const FNavigationPath& Path, TConstArrayV
 			UE::NavCorridor::Private::StringPull(Portals, 0, Portals.Num() - 1, Params.PathOffsetFromBoundaries);
 		}
 	}
+
+#if ENABLE_VISUAL_LOG
+	if (UE::NavCorridor::Private::bEnableNavCorridorVLog)
+	{
+		if (const INavAgentInterface* PathQuerierNavAgent = Cast<const INavAgentInterface>(PathQuerier))
+		{
+			UE_VLOG_LOCATION(PathQuerier, LogNavCorridorPath, Log, PathQuerierNavAgent->GetNavAgentLocation(), 20.f, FColor::Green, TEXT_EMPTY);
+		}
+
+		// Draw path
+		for (int32 PointIndex = 0; PointIndex < PathPoints.Num(); PointIndex++)
+		{
+			if ((PointIndex + 1) < PathPoints.Num())
+			{
+				UE_VLOG_SEGMENT_THICK(PathQuerier, LogNavCorridorPath, Log, PathPoints[PointIndex].Location, PathPoints[PointIndex + 1].Location, FColor(64, 64, 64), 2.f, TEXT_EMPTY);
+			}
+
+			UE_VLOG_LOCATION(PathQuerier, LogNavCorridorPath, Log, PathPoints[PointIndex].Location, 20.f, FColor::White, TEXT("%d-%d%s"), PointIndex, FNavMeshNodeFlags(PathPoints[PointIndex].Flags).AreaFlags, Params.PinnedPointsIndices.Contains(PointIndex) ? TEXT(" Pinned") : TEXT_EMPTY);
+		}
+
+		const FVector Offset(0.f, 0.f, 10.f);
+		for (int32 PortalIndex = 0; PortalIndex < Portals.Num(); PortalIndex++)
+		{
+			const FNavCorridorPortal& Portal = Portals[PortalIndex];
+
+			// Portal segment
+			UE_VLOG_SEGMENT_THICK(PathQuerier, LogNavCorridorPath, Log, Offset + Portal.Left, Offset + Portal.Right, Portal.bIsPathCorner ? FColor::Red : FColor::Orange, Portal.bIsPathCorner ? 2.0f : 1.f, TEXT_EMPTY);
+
+			if ((PortalIndex + 1) < Portals.Num())
+			{
+				// Sector boundaries
+				const FNavCorridorPortal& NextPortal = Portals[PortalIndex + 1];
+				UE_VLOG_SEGMENT_THICK(PathQuerier, LogNavCorridorPath, Log, Offset + Portal.Left, Offset + NextPortal.Left, FColor::Orange, 2.f, TEXT_EMPTY);
+				UE_VLOG_SEGMENT_THICK(PathQuerier, LogNavCorridorPath, Log, Offset + Portal.Right, Offset + NextPortal.Right, FColor::Orange, 2.f, TEXT_EMPTY);
+
+				// Path
+				UE_VLOG_SEGMENT_THICK(PathQuerier, LogNavCorridorPath, Log, Offset + Portal.Location, Offset + NextPortal.Location, FColor(255, 128, 128), 2.f, TEXT_EMPTY);
+			}
+
+			UE_VLOG_LOCATION(PathQuerier, LogNavCorridorPath, Log, Offset + Portal.Location, 5.f, Portal.bIsPinnedLocation ? FColor(75, 75, 75) : FColor::White, TEXT("%d%s"), PortalIndex, Portal.bIsPinnedLocation ? TEXT(" (Pinned)" : TEXT_EMPTY));
+		}
+	}
+#endif // ENABLE_VISUAL_LOG
+
 }
 
 void FNavCorridor::OffsetPathLocationsFromWalls(const float Offset, bool bOffsetFirst, bool bOffsetLast)
diff --git a/Engine/Plugins/Runtime/NavCorridor/Source/NavCorridor/Public/NavCorridor.h b/Engine/Plugins/Runtime/NavCorridor/Source/NavCorridor/Public/NavCorridor.h
--- a/Engine/Plugins/Runtime/NavCorridor/Source/NavCorridor/Public/NavCorridor.h
+++ b/Engine/Plugins/Runtime/NavCorridor/Source/NavCorridor/Public/NavCorridor.h
@@ -22,7 +22,10 @@ struct NAVCORRIDOR_API FNavCorridorParams
 		LargeSectorThreshold = Width;
 		SimplifyEdgeThreshold = Width * 0.1f;
 	}
-	
+
+	/** Points along the path which we want to visit instead of string pulling the location on the shortest path. */
+	TArray<int32, TInlineAllocator<8>> PinnedPointsIndices;
+
 	/** Width of the corridor to build */
 	UPROPERTY(EditAnywhere, Category = "Pathfinding", meta=(ClampMin="10.0"))
 	float Width = 200.0f;
@@ -75,10 +78,13 @@ struct NAVCORRIDOR_API FNavCorridorPortal
 	FVector Location = FVector::ZeroVector;
 	
 	/** Path point index of the original path */
-	int32 PathPointIndex = 0;
+	int16 PathPointIndex = 0;
 	
 	/** True if the portal is at original path point corner. */
-	bool bIsPathCorner = false;
+	bool bIsPathCorner : 1 = false;
+
+	/** When true, the string pulling will be forced to visit the original location at this portal, unless that location is outside the corridor (eg: too close to an obstacle). */
+	bool bIsPinnedLocation : 1 = false;
 };
 
 /**
