Trivial: Fixing wrong argument type passed to UE_VLOG_SEGMENT_THICK in NavCorridor.cpp.

#rnx

diff --git a/Engine/Plugins/Runtime/NavCorridor/Source/NavCorridor/Private/NavCorridor.cpp b/Engine/Plugins/Runtime/NavCorridor/Source/NavCorridor/Private/NavCorridor.cpp
--- a/Engine/Plugins/Runtime/NavCorridor/Source/NavCorridor/Private/NavCorridor.cpp
+++ b/Engine/Plugins/Runtime/NavCorridor/Source/NavCorridor/Private/NavCorridor.cpp
@@ -1480,7 +1480,7 @@ void FNavCorridor::BuildFromPathPoints(const FNavigationPath& Path, TConstArrayV
 		{
 			if ((PointIndex + 1) < PathPoints.Num())
 			{
-				UE_VLOG_SEGMENT_THICK(PathQuerier, LogNavCorridorPath, Log, PathPoints[PointIndex].Location, PathPoints[PointIndex + 1].Location, FColor(64, 64, 64), 2.f, TEXT_EMPTY);
+				UE_VLOG_SEGMENT_THICK(PathQuerier, LogNavCorridorPath, Log, PathPoints[PointIndex].Location, PathPoints[PointIndex + 1].Location, FColor(64, 64, 64), 2, TEXT_EMPTY);
 			}
 
 			UE_VLOG_LOCATION(PathQuerier, LogNavCorridorPath, Log, PathPoints[PointIndex].Location, 20.f, FColor::White, TEXT("%d-%d%s"), PointIndex, FNavMeshNodeFlags(PathPoints[PointIndex].Flags).AreaFlags, Params.PinnedPointsIndices.Contains(PointIndex) ? TEXT(" Pinned") : TEXT_EMPTY);
@@ -1492,17 +1492,17 @@ void FNavCorridor::BuildFromPathPoints(const FNavigationPath& Path, TConstArrayV
 			const FNavCorridorPortal& Portal = Portals[PortalIndex];
 
 			// Portal segment
-			UE_VLOG_SEGMENT_THICK(PathQuerier, LogNavCorridorPath, Log, Offset + Portal.Left, Offset + Portal.Right, Portal.bIsPathCorner ? FColor::Red : FColor::Orange, Portal.bIsPathCorner ? 2.0f : 1.f, TEXT_EMPTY);
+			UE_VLOG_SEGMENT_THICK(PathQuerier, LogNavCorridorPath, Log, Offset + Portal.Left, Offset + Portal.Right, Portal.bIsPathCorner ? FColor::Red : FColor::Orange, Portal.bIsPathCorner ? 2 : 1, TEXT_EMPTY);
 
 			if ((PortalIndex + 1) < Portals.Num())
 			{
 				// Sector boundaries
 				const FNavCorridorPortal& NextPortal = Portals[PortalIndex + 1];
-				UE_VLOG_SEGMENT_THICK(PathQuerier, LogNavCorridorPath, Log, Offset + Portal.Left, Offset + NextPortal.Left, FColor::Orange, 2.f, TEXT_EMPTY);
-				UE_VLOG_SEGMENT_THICK(PathQuerier, LogNavCorridorPath, Log, Offset + Portal.Right, Offset + NextPortal.Right, FColor::Orange, 2.f, TEXT_EMPTY);
+				UE_VLOG_SEGMENT_THICK(PathQuerier, LogNavCorridorPath, Log, Offset + Portal.Left, Offset + NextPortal.Left, FColor::Orange, 2, TEXT_EMPTY);
+				UE_VLOG_SEGMENT_THICK(PathQuerier, LogNavCorridorPath, Log, Offset + Portal.Right, Offset + NextPortal.Right, FColor::Orange, 2, TEXT_EMPTY);
 
 				// Path
-				UE_VLOG_SEGMENT_THICK(PathQuerier, LogNavCorridorPath, Log, Offset + Portal.Location, Offset + NextPortal.Location, FColor(255, 128, 128), 2.f, TEXT_EMPTY);
+				UE_VLOG_SEGMENT_THICK(PathQuerier, LogNavCorridorPath, Log, Offset + Portal.Location, Offset + NextPortal.Location, FColor(255, 128, 128), 2, TEXT_EMPTY);
 			}
 
 			UE_VLOG_LOCATION(PathQuerier, LogNavCorridorPath, Log, Offset + Portal.Location, 5.f, Portal.bIsPinnedLocation ? FColor(75, 75, 75) : FColor::White, TEXT("%d%s"), PortalIndex, Portal.bIsPinnedLocation ? TEXT(" (Pinned)" : TEXT_EMPTY));
