Fix issue with nav corridor SimplifyConcavePortals that could break some portals while simplifying the corridor.

diff --git a/Engine/Plugins/Runtime/NavCorridor/Source/NavCorridor/Private/NavCorridor.cpp b/Engine/Plugins/Runtime/NavCorridor/Source/NavCorridor/Private/NavCorridor.cpp
--- a/Engine/Plugins/Runtime/NavCorridor/Source/NavCorridor/Private/NavCorridor.cpp
+++ b/Engine/Plugins/Runtime/NavCorridor/Source/NavCorridor/Private/NavCorridor.cpp
@@ -28,6 +28,10 @@ namespace UE::NavCorridor::Private
 	static FAutoConsoleVariableRef CVarNavCorridorDebugEnableVisualLogging(TEXT("ai.nav.EnableNavCorridorVLog"), bEnableNavCorridorVLog, TEXT("If set, NavCorridor debug information will be logged in the Visual Logger."), ECVF_Cheat);
 #endif // ENABLE_VISUAL_LOG
 
+	static bool bUseSimplifyConcavePortalsFix = true;
+	static FAutoConsoleVariableRef CVarUseSimplifyConcavePortalsFix(TEXT("ai.nav.bUseSimplifyConcavePortalsFix"),
+		bUseSimplifyConcavePortalsFix, TEXT("True by default. Set to false to revert to previous corridor behavior when simplifying."), ECVF_Cheat);
+	
 	/** Epsilon used when testing if 3 points for a convex corner.
 	 * The check is done using signed distance of the corner points to the segment from first to last points.
 	 * This is how much it can be concave. */
@@ -996,6 +1000,12 @@ namespace UE::NavCorridor::Private
 						UE::AI::IntersectLineLine2D(Portal1.Left, Portal1.Left + Dir, Portal0.Left, Portal0.Right, Temp, T0);
 						UE::AI::IntersectLineLine2D(Portal1.Left, Portal1.Left + Dir, Portal2.Left, Portal2.Right, Temp, T2);
 
+						if (bUseSimplifyConcavePortalsFix && (T0 < 0 || T2 < 0 || T0 > 1 || T2 > 1))
+						{
+							// Skip simplify this portal to avoid degenerating the path.
+							continue;
+						}
+
 						const FVector AdjustedLeft0 = FMath::Lerp(Portal0.Left, Portal0.Right, T0);
 						const FVector AdjustedLeft2 = FMath::Lerp(Portal2.Left, Portal2.Right, T2);
 
@@ -1010,6 +1020,12 @@ namespace UE::NavCorridor::Private
 						UE::AI::IntersectLineLine2D(Portal1.Right, Portal1.Right + Dir, Portal0.Left, Portal0.Right, Temp, T0);
 						UE::AI::IntersectLineLine2D(Portal1.Right, Portal1.Right + Dir, Portal2.Left, Portal2.Right, Temp, T2);
 
+						if (bUseSimplifyConcavePortalsFix && (T0 < 0 || T2 < 0 || T0 > 1 || T2 > 1))
+						{
+							// Skip simplify this portal to avoid degenerating the path.
+							continue;
+						}
+
 						const FVector AdjustedRight0 = FMath::Lerp(Portal0.Left, Portal0.Right, T0);
 						const FVector AdjustedRight2 = FMath::Lerp(Portal2.Left, Portal2.Right, T2);
 
@@ -1315,6 +1331,27 @@ namespace UE::NavCorridor::Private
 
 }; // UE::NavCorridor::Private
 
+
+//-------------------------------------------------------
+// FNavCorridorParams
+//-------------------------------------------------------
+
+FString FNavCorridorParams::ToString() const
+{
+	return FString::Printf(TEXT("Width: %.1f, PathOffsetFromBoundaries: %.1f, ObstacleTaperAngle: %.1f, SmallSectorThreshold %.1f, " 
+		"LargeSectorThreshold: %.1f, SimplifyEdgeThreshold: %.1f, bSimplifyFlipPortals: %i, bSimplifyConvexPortals: %i, bSimplifyConcavePortals: %i"),
+		Width,
+		PathOffsetFromBoundaries,
+		ObstacleTaperAngle,
+		SmallSectorThreshold,
+		LargeSectorThreshold,
+		SimplifyEdgeThreshold,
+		bSimplifyFlipPortals,
+		bSimplifyConvexPortals,
+		bSimplifyConcavePortals);
+}
+
+
 //-------------------------------------------------------
 // FNavCorridor
 //-------------------------------------------------------
diff --git a/Engine/Plugins/Runtime/NavCorridor/Source/NavCorridor/Public/NavCorridor.h b/Engine/Plugins/Runtime/NavCorridor/Source/NavCorridor/Public/NavCorridor.h
--- a/Engine/Plugins/Runtime/NavCorridor/Source/NavCorridor/Public/NavCorridor.h
+++ b/Engine/Plugins/Runtime/NavCorridor/Source/NavCorridor/Public/NavCorridor.h
@@ -61,6 +61,9 @@ struct NAVCORRIDOR_API FNavCorridorParams
 	/** If true do concave portals simplification. */
 	UPROPERTY(EditAnywhere, Category = "Pathfinding")
 	bool bSimplifyConcavePortals = true;
+
+	/** Returns a string of the corridor parameters. */
+	FString ToString() const;
 };
 
 /**
